Texture: fix many issues with Required Memory Estimate
was wrong for latlong cubes, astc, MaxTextureSize
factor out ComputeLongLatCubemapExtents to remove duplication of cube size logic
fix GetOutputMipInfo was wrong on latlong cubes with pow2 options
fix "VT has blocks with mismatched aspect ratios" warning could be thrown incorrectly on non-VT non-UDIM textures with "resize to specific" set

diff --git a/Engine/Source/Developer/TextureBuildUtilities/Private/TextureBuildUtilities.cpp b/Engine/Source/Developer/TextureBuildUtilities/Private/TextureBuildUtilities.cpp
--- a/Engine/Source/Developer/TextureBuildUtilities/Private/TextureBuildUtilities.cpp
+++ b/Engine/Source/Developer/TextureBuildUtilities/Private/TextureBuildUtilities.cpp
@@ -408,28 +408,26 @@ static int GetWithinSliceRDOMemoryUsePerPixel(EPixelFormat PixelFormat)
 	}
 }
 
-static int64 CalculateTextureSourceBytesFromImageInfo(FImageInfo InImageInfo, int32 InMipCount, bool bInVolume)
+static int64 CalculateTextureSourceBytesFromImageInfo(const FImageInfo & InImageInfoConst, const int32 InMipCountConst, const bool bInVolume)
 {
-	if (InMipCount > 0)
+	FImageInfo ImageInfo = InImageInfoConst;
+	int32 MipCount = InMipCountConst;
+
+	int64 SourceBytes = 0;
+	while (MipCount)
 	{
-		int64 SourceBytes = 0;
-		while (InMipCount)
-		{
-			SourceBytes += InImageInfo.GetImageSizeBytes();
+		SourceBytes += ImageInfo.GetImageSizeBytes();
 
-			InImageInfo.SizeX = FEncodedTextureDescription::GetMipWidth(InImageInfo.SizeX, 1);
-			InImageInfo.SizeY = FEncodedTextureDescription::GetMipHeight(InImageInfo.SizeY, 1);
-			if (bInVolume)
-			{
-				InImageInfo.NumSlices = FEncodedTextureDescription::GetMipDepth(InImageInfo.NumSlices, 1, true);
-			}
-			InMipCount--;
+		ImageInfo.SizeX = FEncodedTextureDescription::GetMipWidth(ImageInfo.SizeX, 1);
+		ImageInfo.SizeY = FEncodedTextureDescription::GetMipHeight(ImageInfo.SizeY, 1);
+		if (bInVolume)
+		{
+			ImageInfo.NumSlices = FEncodedTextureDescription::GetMipDepth(ImageInfo.NumSlices, 1, true);
 		}
-
-		return SourceBytes;
+		MipCount--;
 	}
 
-	return -1;
+	return SourceBytes;
 }
 
 
@@ -480,7 +478,7 @@ TEXTUREBUILDUTILITIES_API int64 GetVirtualTextureRequiredMemoryEstimate(const FT
 	TConstArrayView<UE::TextureBuildUtilities::FVirtualTextureSourceBlockInfo> InSourceBlocks)
 {
 	const bool bRDO = true;
-	// @todo Oodle : be careful about using BuildSettings for this as there are two buildsettingses, just assume its on for now
+	// @todo Oodle : be careful about using BuildSettings for bRDO as there are two buildsettingses, just assume its on for now
 	//   <- FIX ME, allow lower mem estimates for non-RDO
 
 	// over-estimate is okay
@@ -505,6 +503,8 @@ TEXTUREBUILDUTILITIES_API int64 GetVirtualTextureRequiredMemoryEstimate(const FT
 	int64 TotalTopMipNumPixelsPerLayer = 0;
 	int64 LargestBlockTopMipNumPixels = 0;
 
+	int64 ResizingPhaseMemUsePerLayer = 0;
+
 	// All layers in a VT must have the same layout for each block. Layers only can change source pixels + format, not dims.
 	for (int32 BlockIndex = 0; BlockIndex < InSourceBlocks.Num(); ++BlockIndex)
 	{
@@ -526,11 +526,42 @@ TEXTUREBUILDUTILITIES_API int64 GetVirtualTextureRequiredMemoryEstimate(const FT
 		// assume pow2 options are the same for all layers, just use layer 0 here :
 		const FTextureBuildSettings& LayerBuildSettings = InBuildSettingsPerLayer[0];
 
+		check( ! LayerBuildSettings.bVolume );
+		check( ! LayerBuildSettings.bCubemap );
+		check( ! LayerBuildSettings.bLongLatSource );
+
 		int32 TargetSizeX, TargetSizeY, TargetSizeZ;
-		UE::TextureBuildUtilities::GetPowerOfTwoTargetTextureSize(SourceBlock.SizeX, SourceBlock.SizeY, SourceBlock.NumSlices,
+		bool bDidPow2 = UE::TextureBuildUtilities::GetPowerOfTwoTargetTextureSize(SourceBlock.SizeX, SourceBlock.SizeY, SourceBlock.NumSlices,
 			LayerBuildSettings.bVolume, (ETexturePowerOfTwoSetting::Type)LayerBuildSettings.PowerOfTwoMode,
 			LayerBuildSettings.ResizeDuringBuildX, LayerBuildSettings.ResizeDuringBuildY,
 			TargetSizeX, TargetSizeY, TargetSizeZ);
+		
+		int64 AfterPow2TopMipNumPixels = (int64) TargetSizeX * TargetSizeY * TargetSizeZ;
+			
+		// MaxTextureSize on UDIM applies to each block on its own
+		if ( LayerBuildSettings.MaxTextureResolution != TNumericLimits<uint32>::Max() )
+		{
+			// max memory use of the MaxTextureResolution op is the source in RGBA32F + a mip of size /2 in RGBA32F
+			ResizingPhaseMemUsePerLayer += AfterPow2TopMipNumPixels * 20;  // (1 + 1/4) * 16;
+		
+			// ResizingPhaseMemUse is per Layer at this point
+
+			while( (uint32)TargetSizeX > LayerBuildSettings.MaxTextureResolution || (uint32)TargetSizeY > LayerBuildSettings.MaxTextureResolution )
+			{
+				TargetSizeX = FMath::Max(TargetSizeX>>1,1);
+				TargetSizeY = FMath::Max(TargetSizeY>>1,1);
+
+				// TargetSizeZ not changed
+				check( ! LayerBuildSettings.bVolume );
+			}
+		}
+		else if ( bDidPow2 )
+		{
+			// surface in in RGBA32F
+			ResizingPhaseMemUsePerLayer += AfterPow2TopMipNumPixels * 16;
+		}
+
+		// TargetSize is now the size after Pow2 and MaxTextureSize resizes :
 
 		int64 CurrentBlockTopMipNumPixels = (int64)TargetSizeX * TargetSizeY * TargetSizeZ;
 
@@ -544,6 +575,10 @@ TEXTUREBUILDUTILITIES_API int64 GetVirtualTextureRequiredMemoryEstimate(const FT
 		return -1; /* Unknown */
 	}
 
+	int64 ResizingPhaseMemUse = TotalSourceBytes + ResizingPhaseMemUsePerLayer * InLayerFormats.Num();
+
+	// after this point, "numpixels" is the number encode to VT and output pixel format
+
 	// assume full mip chain :
 	int64 TotalPixelsPerLayer = (TotalTopMipNumPixelsPerLayer * 4) / 3;
 
@@ -642,9 +677,12 @@ TEXTUREBUILDUTILITIES_API int64 GetVirtualTextureRequiredMemoryEstimate(const FT
 	// take larger of mem use during float image filter phase or tile compression phase
 	MemoryEstimate += FMath::Max(IntermediateFloatColorBytes, TileCompressionBytes + OutputSizeBytes);
 
+	// larger of early resize phase and VT build phase :
+	MemoryEstimate = FMath::Max(ResizingPhaseMemUse,MemoryEstimate);
+
 	MemoryEstimate += 64 * 1024; // overhead room
 
-	//UE_LOG(LogTexture,Display,TEXT("GetBuildRequiredMemoryEstimate VT : %.3f MB"),MemoryEstimate/(1024*1024.f));
+	//UE_LOG(LogTextureBuildUtilities,Display,TEXT("GetBuildRequiredMemoryEstimate VT : %.3f MB"),MemoryEstimate/(1024*1024.f));
 
 	return MemoryEstimate;
 }
@@ -654,7 +692,7 @@ TEXTUREBUILDUTILITIES_API int64 GetVirtualTextureRequiredMemoryEstimate(const FT
 TEXTUREBUILDUTILITIES_API int64 GetPhysicalTextureBuildMemoryEstimate(const FTextureBuildSettings* InSettingsPerLayerFetchFirst, const FImageInfo& InSourceImageInfo, int32 InMipCount)
 {
 	const bool bRDO = true;
-	// @todo Oodle : be careful about using BuildSettings for this as there are two buildsettingses, just assume its on for now
+	// @todo Oodle : be careful about using BuildSettings for bRDO as there are two buildsettingses, just assume its on for now
 	//   <- FIX ME, allow lower mem estimates for non-RDO
 
 	// over-estimate is okay
@@ -662,151 +700,238 @@ TEXTUREBUILDUTILITIES_API int64 GetPhysicalTextureBuildMemoryEstimate(const FTex
 
 	int64 MaxNumberOfWorkers = FMath::Max(1, FTaskGraphInterface::Get().GetNumWorkerThreads());
 
+	const FTextureBuildSettings& BuildSettings = InSettingsPerLayerFetchFirst[0];
+	// non VT
+
+	// Compute the memory it should take to uncompress the bulkdata in memory
+	int64 TotalSourceBytes = CalculateTextureSourceBytesFromImageInfo(InSourceImageInfo, InMipCount, BuildSettings.bVolume);
+	if (TotalSourceBytes <= 0)
 	{
-		const FTextureBuildSettings& BuildSettings = InSettingsPerLayerFetchFirst[0];
-		// non VT
+		return -1; /* Unknown */
+	}
 
-		// Compute the memory it should take to uncompress the bulkdata in memory
-		int64 TotalSourceBytes = CalculateTextureSourceBytesFromImageInfo(InSourceImageInfo, InMipCount, BuildSettings.bVolume);
-		if (TotalSourceBytes <= 0)
+	// NOTE: it would be ideal to call Texture::GetBuiltTextureSize here, but we don't have a Texture pointer, sigh.
+
+	int32 TargetSizeX, TargetSizeY, TargetSizeZ;
+	bool bDidPow2 = UE::TextureBuildUtilities::GetPowerOfTwoTargetTextureSize(InSourceImageInfo.SizeX, InSourceImageInfo.SizeY, InSourceImageInfo.NumSlices,
+		BuildSettings.bVolume, (ETexturePowerOfTwoSetting::Type)BuildSettings.PowerOfTwoMode,
+		BuildSettings.ResizeDuringBuildX, BuildSettings.ResizeDuringBuildY,
+		TargetSizeX, TargetSizeY, TargetSizeZ);
+	
+	// mem use of the Pow2 op is always lower that later ops, so no need to compute it
+
+	int64 ResizingPhaseMemUse = TotalSourceBytes;
+	int64 InitialTopMipNumPixels = (int64) TargetSizeX * TargetSizeY * TargetSizeZ;
+	ResizingPhaseMemUse += InitialTopMipNumPixels * 16; // top mip in RGBA32F may be needed
+
+	if ( BuildSettings.bLongLatSource )
+	{
+		// longlat to cube is after pow2 pad
+		TargetSizeX = TargetSizeY = ComputeLongLatCubemapExtents(TargetSizeX,BuildSettings.MaxTextureResolution);
+
+		// could be a cube array :
+		TargetSizeZ = InSourceImageInfo.NumSlices * 6;
+
+		// mem use of the longlat->cube operation; requires source longlat in RGBA32F and the cube output :
+		ResizingPhaseMemUse += (int64) TargetSizeX * TargetSizeY * TargetSizeZ * 16;
+	}
+	else if ( BuildSettings.MaxTextureResolution != TNumericLimits<uint32>::Max() )
+	{
+		// apply MaxTextureResolution
+		// NOTE: it would be ideal to call Texture::GetBuiltTextureSize here, but we don't have a Texture pointer, sigh.
+		//		(or some kind of shared function rather than duplicating all this logic)
+
+		// max memory use of the MaxTextureResolution op is the source in RGBA32F + a mip of size /2 in RGBA32F
+		ResizingPhaseMemUse += (InitialTopMipNumPixels/4) * 16;
+		
+		while( (uint32)TargetSizeX > BuildSettings.MaxTextureResolution || (uint32)TargetSizeY > BuildSettings.MaxTextureResolution )
 		{
-			return -1; /* Unknown */
+			TargetSizeX = FMath::Max(TargetSizeX>>1,1);
+			TargetSizeY = FMath::Max(TargetSizeY>>1,1);
+
+			if ( BuildSettings.bVolume )
+			{
+				TargetSizeZ = FMath::Max(TargetSizeZ>>1,1);
+			}
 		}
+	}
 
+	// from here on, NumPixels is the number that will be encoded to the output pixel format
 
+	int64 TotalTopMipNumPixels = (int64)TargetSizeX * TargetSizeY * TargetSizeZ;
 
-		int32 TargetSizeX, TargetSizeY, TargetSizeZ;
-		UE::TextureBuildUtilities::GetPowerOfTwoTargetTextureSize(InSourceImageInfo.SizeX, InSourceImageInfo.SizeY, InSourceImageInfo.NumSlices,
-			BuildSettings.bVolume, (ETexturePowerOfTwoSetting::Type)BuildSettings.PowerOfTwoMode,
-			BuildSettings.ResizeDuringBuildX, BuildSettings.ResizeDuringBuildY,
-			TargetSizeX, TargetSizeY, TargetSizeZ);
+	// assume full mip chain :
+	//	(volume mips are smaller than this, but over-estimating is okay)
+	int64 TotalNumPixels = (TotalTopMipNumPixels * 4) / 3;
 
-		int64 TotalTopMipNumPixels = (int64)TargetSizeX * TargetSizeY * TargetSizeZ;
+	// actually we have each mip twice for the float image filter phase so this is under-counting
+	//	but that isn't held allocated while the output is made, so it can overlap with that mem
+	int64 IntermediateFloatColorBytes = TotalNumPixels * sizeof(FLinearColor);
 
-		// assume full mip chain :
-		int64 TotalNumPixels = (TotalTopMipNumPixels * 4) / 3;
+	// if we knew the source BulkData was always freed during encoding, TotalSourceBytes could be dropped
+	int64 MemoryEstimate = TotalSourceBytes + IntermediateFloatColorBytes;
 
-		// actually we have each mip twice for the float image filter phase so this is under-counting
-		//	but that isn't held allocated while the output is made, so it can overlap with that mem
-		int64 IntermediateFloatColorBytes = TotalNumPixels * sizeof(FLinearColor);
+	// Assume alpha exists if we don't know for worst-case handling.
+	const bool bHasAlphaFallback = true;
+	EPixelFormat PixelFormat = GetOutputPixelFormatWithFallback(BuildSettings, bHasAlphaFallback);
 
-		int64 MemoryEstimate = TotalSourceBytes + IntermediateFloatColorBytes;
+	if (PixelFormat == PF_Unknown)
+	{
+		return -1; /* Unknown */
+	}
 
-		// Assume alpha exists if we don't know for worst-case handling.
-		EPixelFormat PixelFormat = GetOutputPixelFormatWithFallback(BuildSettings, true);
+	const FPixelFormatInfo& PFI = GPixelFormats[PixelFormat];
 
-		if (PixelFormat == PF_Unknown)
+	const int64 OutputSizeBytes = (TotalNumPixels * PFI.BlockBytes) / (PFI.BlockSizeX * PFI.BlockSizeY);
+
+	MemoryEstimate += OutputSizeBytes;
+
+	// check to see if it's uncompressed or a BCN format :
+	if (IsDXTCBlockCompressedTextureFormat(PixelFormat))
+	{
+		// block-compressed format ; assume it's using Oodle Texture
+
+		if (bRDO)
 		{
-			return -1; /* Unknown */
+			// two more copies in outputsize
+			// baseline encode + UT or Layout
+			MemoryEstimate += OutputSizeBytes * 2;
 		}
 
-		const FPixelFormatInfo& PFI = GPixelFormats[PixelFormat];
-
-		const int64 OutputSizeBytes = (TotalNumPixels * PFI.BlockBytes) / (PFI.BlockSizeX * PFI.BlockSizeY);
+		// you also have to convert the float surface to an input format for Oodle
+		//	this copy is done in TFO
+		//  Oodle then allocs another copy to swizzle into blocks before encoding
 
-		MemoryEstimate += OutputSizeBytes;
+		int IntermediateBytesPerPixel;
+		bool bNeedsIntermediateCopy = true;
 
-		// check to see if it's uncompressed or a BCN format :
-		if (IsDXTCBlockCompressedTextureFormat(PixelFormat))
+		// this matches the logic in TextureFormatOodle :
+		if (PixelFormat == PF_BC6H)
 		{
-			// block-compressed format ; assume it's using Oodle Texture
+			IntermediateBytesPerPixel = 16; //RGBAF32
+			bNeedsIntermediateCopy = false; // no intermediate used in TFO (float source kept), 1 blocksurf
+		}
+		else if (PixelFormat == PF_BC4 || PixelFormat == PF_BC5)
+		{
+			// changed: TFO uses 2_U16 now (4 byte intermediate)
+			IntermediateBytesPerPixel = 8; // RGBA16
+		}
+		else
+		{
+			IntermediateBytesPerPixel = 4; // RGBA8
+		}
 
-			if (bRDO)
-			{
-				// two more copies in outputsize
-				// baseline encode + UT or Layout
-				MemoryEstimate += OutputSizeBytes * 2;
-			}
+		int NumIntermediateCopies = 1; // BlockSurf
+		if (bNeedsIntermediateCopy) NumIntermediateCopies++;
 
-			// you also have to convert the float surface to an input format for Oodle
-			//	this copy is done in TFO
-			//  Oodle then allocs another copy to swizzle into blocks before encoding
+		MemoryEstimate += NumIntermediateCopies * IntermediateBytesPerPixel * TotalNumPixels;
 
-			int IntermediateBytesPerPixel;
-			bool bNeedsIntermediateCopy = true;
+		if (bRDO)
+		{
+			// activity map for whole image :
+			// (this has changed in newer versions of Oodle Texture)
 
-			// this matches the logic in TextureFormatOodle :
-			if (PixelFormat == PF_BC6H)
-			{
-				IntermediateBytesPerPixel = 16; //RGBAF32
-				bNeedsIntermediateCopy = false; // no intermediate used in TFO (float source kept), 1 blocksurf
-			}
-			else if (PixelFormat == PF_BC4 || PixelFormat == PF_BC5)
-			{
-				// changed: TFO uses 2_U16 now (4 byte intermediate)
-				IntermediateBytesPerPixel = 8; // RGBA16
-			}
-			else
-			{
-				IntermediateBytesPerPixel = 4; // RGBA8
-			}
+			// Phase1 = computing activity map
+			int ActivityBytesPerPixel;
 
-			int NumIntermediateCopies = 1; // BlockSurf
-			if (bNeedsIntermediateCopy) NumIntermediateCopies++;
+			if (PixelFormat == PF_BC4) ActivityBytesPerPixel = 12;
+			else if (PixelFormat == PF_BC5) ActivityBytesPerPixel = 16;
+			else ActivityBytesPerPixel = 24;
 
-			MemoryEstimate += NumIntermediateCopies * IntermediateBytesPerPixel * TotalNumPixels;
+			int64 RDOPhase1MemUse = ActivityBytesPerPixel * TotalNumPixels;
 
-			if (bRDO)
+			// Phase2 = cut into slices, encode each slice
+			// per-slice data structure memory use
+			// non-RDO is all on stack so zero
+
+			// fewer workers for small images ; roughly one slice per 64 KB of output
+			//int64 NumberofSlices = FMath::DivideAndRoundUp<int64>(OutputSizeBytes,64*1024);
+			int64 PixelsPerSlice = (64 * 1024 * TotalNumPixels) / OutputSizeBytes;
+			int64 NumberofSlices = FMath::DivideAndRoundUp<int64>(TotalNumPixels, PixelsPerSlice);
+			if (NumberofSlices <= 4)
 			{
-				// activity map for whole image :
-				// (this has changed in newer versions of Oodle Texture)
+				PixelsPerSlice = TotalNumPixels / NumberofSlices;
+			}
 
-				// Phase1 = computing activity map
-				int ActivityBytesPerPixel;
+			int64 MemoryUsePerWorker = PixelsPerSlice * GetWithinSliceRDOMemoryUsePerPixel(PixelFormat);
+				// MemoryUsePerWorker is around 10 MB
+			int64 NumberOfWorkers = FMath::Min(NumberofSlices, MaxNumberOfWorkers);
 
-				if (PixelFormat == PF_BC4) ActivityBytesPerPixel = 12;
-				else if (PixelFormat == PF_BC5) ActivityBytesPerPixel = 16;
-				else ActivityBytesPerPixel = 24;
+			int64 RDOPhase2MemUse = 4 * TotalNumPixels; // activity map held on whole image
+			RDOPhase2MemUse += NumberOfWorkers * MemoryUsePerWorker;
 
-				int64 RDOPhase1MemUse = ActivityBytesPerPixel * TotalNumPixels;
+			// usually phase2 is higher
+			// but on large BC6 images on machines with low core counts, phase1 can be higher
 
-				// Phase2 = cut into slices, encode each slice
-				// per-slice data structure memory use
-				// non-RDO is all on stack so zero
+			MemoryEstimate += FMath::Max(RDOPhase1MemUse, RDOPhase2MemUse);
+		}
+	}
+	else if (IsASTCBlockCompressedTextureFormat(PixelFormat))
+	{
+		// ASTCenc does an entermediate copy to RGBA16F for HDR formats and RGBA8 for LDR
+		MemoryEstimate += (IsHDR(PixelFormat) ? 8 : 4) * TotalNumPixels;
+			
+		// internal memory use of ASTCenc :
+		//	measured from command line astcenc.exe
+		MemoryEstimate += 10 * TotalNumPixels;
+	}
+	else if ( PFI.BlockSizeX > 1 )
+	{
+		// block compressed but not Oodle or ASTC (eg. ETC)
+		// note: memory ues of non-Oodle encoders is not estimated
+		// @todo : fix me
+			
+		// prefer over-estimate to under-estimate :
+		MemoryEstimate += 16 * TotalNumPixels;
+	}
+	else
+	{
+		// non-blocked encoder (uncompressed)
+			
+		// some of the TextureFormatUncompressed encoders use a scratch image
+		//	must over-estimate to be safe :
+		MemoryEstimate += 4 * TotalNumPixels;
+	}
 
-				// fewer workers for small images ; roughly one slice per 64 KB of output
-				//int64 NumberofSlices = FMath::DivideAndRoundUp<int64>(OutputSizeBytes,64*1024);
-				int64 PixelsPerSlice = (64 * 1024 * TotalNumPixels) / OutputSizeBytes;
-				int64 NumberofSlices = FMath::DivideAndRoundUp<int64>(TotalNumPixels, PixelsPerSlice);
-				if (NumberofSlices <= 4)
-				{
-					PixelsPerSlice = TotalNumPixels / NumberofSlices;
-				}
+	// mem use is the max of the phases :
+	MemoryEstimate = FMath::Max(MemoryEstimate,ResizingPhaseMemUse);
 
-				int64 MemoryUsePerWorker = PixelsPerSlice * GetWithinSliceRDOMemoryUsePerPixel(PixelFormat);
-					// MemoryUsePerWorker is around 10 MB
-				int64 NumberOfWorkers = FMath::Min(NumberofSlices, MaxNumberOfWorkers);
+	MemoryEstimate += 64 * 1024; // overhead room
 
-				int64 RDOPhase2MemUse = 4 * TotalNumPixels; // activity map held on whole image
-				RDOPhase2MemUse += NumberOfWorkers * MemoryUsePerWorker;
+	//UE_LOG(LogTextureBuildUtilities,Display,TEXT("GetBuildRequiredMemoryEstimate non-VT : %.3f MB"),MemoryEstimate/(1024*1024.f));
 
-				// usually phase2 is higher
-				// but on large BC6 images on machines with low core counts, phase1 can be higher
+	return MemoryEstimate;
 
-				MemoryEstimate += FMath::Max(RDOPhase1MemUse, RDOPhase2MemUse);
-			}
-		}
-		else if (IsASTCBlockCompressedTextureFormat(PixelFormat))
-		{
-			// ASTCenc does an entermediate copy to RGBA16F for HDR formats and RGBA8 for LDR
-			MemoryEstimate += (IsHDR(PixelFormat) ? 8 : 4) * TotalNumPixels;
-			// internal memory use of ASTCenc is not estimated
-			// @todo : fix me
-		}
-		else
-		{
-			// note: memory ues of non-Oodle encoders is not estimated
-			// @todo : fix me
-		}
+	// @todo Oodle : not right with Composite 
+	
+	// this is not right for CPU textures, but it is an over-estimate, so that's okay
 
-		MemoryEstimate += 64 * 1024; // overhead room
+	// note: this is intended to be right for TFO, not OTF
+	//	 the cloud TBW and ContentWorker runs that really care about mem use limitations are TFO only
+}
 
-		//UE_LOG(LogTexture,Display,TEXT("GetBuildRequiredMemoryEstimate non-VT : %.3f MB"),MemoryEstimate/(1024*1024.f));
+uint32 ComputeLongLatCubemapExtents(int32 SrcImageSizeX, uint32 MaxCubemapTextureResolution)
+{
+	// MaxTextureSize of 0 is changed to max when filling BuildSettings
+	if ( MaxCubemapTextureResolution == 0 )
+	{
+		MaxCubemapTextureResolution = TNumericLimits<uint32>::Max();
+	}
 
-		return MemoryEstimate;
+	uint32 Out = 1U << FMath::FloorLog2(SrcImageSizeX / 2);
 
-		// @todo Oodle : not right for volumes & latlong cubes
-		// @todo Oodle : not right with Composite , CPU textures
+	if ( Out <= 32 || MaxCubemapTextureResolution <= 32 )
+	{
+		return 32;
+	}
+	else if ( Out > MaxCubemapTextureResolution )
+	{
+		// RoundDownToPowerOfTwo
+		return 1U << FMath::FloorLog2(MaxCubemapTextureResolution);
+	}
+	else
+	{
+		return Out;
 	}
 }
 
diff --git a/Engine/Source/Developer/TextureBuildUtilities/Public/TextureBuildUtilities.h b/Engine/Source/Developer/TextureBuildUtilities/Public/TextureBuildUtilities.h
--- a/Engine/Source/Developer/TextureBuildUtilities/Public/TextureBuildUtilities.h
+++ b/Engine/Source/Developer/TextureBuildUtilities/Public/TextureBuildUtilities.h
@@ -123,6 +123,9 @@ struct TEXTUREBUILDUTILITIES_API FVirtualTextureSourceBlockInfo
 TEXTUREBUILDUTILITIES_API int64 GetVirtualTextureRequiredMemoryEstimate(const FTextureBuildSettings* InBuildSettingsPerLayer,
 	TConstArrayView<ERawImageFormat::Type> InLayerFormats,
 	TConstArrayView<UE::TextureBuildUtilities::FVirtualTextureSourceBlockInfo> InSourceBlocks);
+	
+// ComputeLongLatCubemapExtents is done after pad-to-pow2
+TEXTUREBUILDUTILITIES_API uint32 ComputeLongLatCubemapExtents(int32 SrcImageSizeX, uint32 MaxCubemapTextureResolution);
 
 } // namespace TextureBuildUtilities
 } // namespace UE
diff --git a/Engine/Source/Developer/TextureCompressor/Private/TextureCompressorModule.cpp b/Engine/Source/Developer/TextureCompressor/Private/TextureCompressorModule.cpp
--- a/Engine/Source/Developer/TextureCompressor/Private/TextureCompressorModule.cpp
+++ b/Engine/Source/Developer/TextureCompressor/Private/TextureCompressorModule.cpp
@@ -2251,27 +2251,6 @@ static inline FVector ComputeWSCubeDirectionAtTexelCenter(uint32 CubemapFace, ui
 	return DirectionWS;
 }
 
-static uint32 ComputeLongLatCubemapExtents(int32 SrcImageSizeX, const uint32 MaxCubemapTextureResolution)
-{
-	// MaxCubemapTextureResolution when not set is 0xFFFFFFFF
-
-	uint32 Out = 1U << FMath::FloorLog2(SrcImageSizeX / 2);
-
-	if ( Out <= 32 || MaxCubemapTextureResolution <= 32 )
-	{
-		return 32;
-	}
-	else if ( Out > MaxCubemapTextureResolution )
-	{
-		// RoundDownToPowerOfTwo
-		return 1U << FMath::FloorLog2(MaxCubemapTextureResolution);
-	}
-	else
-	{
-		return Out;
-	}
-}
-
 void ITextureCompressorModule::GenerateBaseCubeMipFromLongitudeLatitude2D(FImage* OutMip, const FImage& SrcImage, const uint32 MaxCubemapTextureResolution, uint8 SourceEncodingOverride)
 {
 	FTextureBuildSettings TempBuildSettings;
@@ -2290,7 +2269,7 @@ void ITextureCompressorModule::GenerateBaseCubeMipFromLongitudeLatitude2D(FImage
 	LinearizeToWorkingColorSpace(SrcImage, LongLatImage, InBuildSettings);
 
 	// TODO_TEXTURE: Expose target size to user.
-	uint32 Extent = ComputeLongLatCubemapExtents(LongLatImage.SizeX, InBuildSettings.MaxTextureResolution);
+	uint32 Extent = UE::TextureBuildUtilities::ComputeLongLatCubemapExtents(LongLatImage.SizeX, InBuildSettings.MaxTextureResolution);
 	float InvExtent = 1.0f / (float)Extent;
 	OutMip->Init(Extent, Extent, SrcImage.NumSlices * 6, ERawImageFormat::RGBA32F, EGammaSpace::Linear);
 
@@ -3707,33 +3686,31 @@ bool FTextureBuildSettings::GetOutputMipInfo(
 		return true;
 	}
 
-	// AFAICT LatLongCubeMaps don't do any of this - pow2 is broken with them but it runs, and max texture stuff
-	// is handled internally in the extents function.
 	int32 BaseSizeX = InMip0SizeX;
 	int32 BaseSizeY = InMip0SizeY;
 	int32 BaseSizeZ = this->bVolume ? InMip0NumSlices : 1; // Volume textures are the only type that mip their Z, arrays and cubes are fixed.
+	
+	ETexturePowerOfTwoSetting::Type PowerOfTwoModeLocal = (ETexturePowerOfTwoSetting::Type)this->PowerOfTwoMode;
+	if (this->MipGenSettings != TMGS_LeaveExistingMips &&
+		PowerOfTwoModeLocal != ETexturePowerOfTwoSetting::None)
+	{
+		int32 TargetSizeX, TargetSizeY, TargetSizeZ;
+		bool NeedsAdjustment = UE::TextureBuildUtilities::GetPowerOfTwoTargetTextureSize(BaseSizeX, BaseSizeY, BaseSizeZ, this->bVolume, PowerOfTwoModeLocal, this->ResizeDuringBuildX, this->ResizeDuringBuildY, TargetSizeX, TargetSizeY, TargetSizeZ);
+		if (NeedsAdjustment)
+		{
+			// In this case we are regenerating the entire mip chain.
+			InExistingMipCount = 1;
+			BaseSizeX = TargetSizeX;
+			BaseSizeY = TargetSizeY;
+			BaseSizeZ = TargetSizeZ; // volume textures already accounted for
+		}
+		// Otherwise we have valid pow2 so we can reuse any existing mips and regenerate
+		// any missing tail mips.
+	}
 
 	// LatLong sources are clamped in ComputeLongLatCubemapExtents
 	if (this->bLongLatSource == false)
 	{
-		ETexturePowerOfTwoSetting::Type PowerOfTwoModeLocal = (ETexturePowerOfTwoSetting::Type)this->PowerOfTwoMode;
-		if (this->MipGenSettings != TMGS_LeaveExistingMips &&
-			PowerOfTwoModeLocal != ETexturePowerOfTwoSetting::None)
-		{
-			int32 TargetSizeX, TargetSizeY, TargetSizeZ;
-			bool NeedsAdjustment = UE::TextureBuildUtilities::GetPowerOfTwoTargetTextureSize(BaseSizeX, BaseSizeY, BaseSizeZ, this->bVolume, PowerOfTwoModeLocal, this->ResizeDuringBuildX, this->ResizeDuringBuildY, TargetSizeX, TargetSizeY, TargetSizeZ);
-			if (NeedsAdjustment)
-			{
-				// In this case we are regenerating the entire mip chain.
-				InExistingMipCount = 1;
-				BaseSizeX = TargetSizeX;
-				BaseSizeY = TargetSizeY;
-				BaseSizeZ = TargetSizeZ; // volume textures already accounted for
-			}
-			// Otherwise we have valid pow2 so we can reuse any existing mips and regenerate
-			// any missing tail mips.
-		}
-
 		// Max texture resolution strips off mips that are above the limit.
 		//int64 MaxTextureResolution = this->MaxTextureResolution; ... ? why was this getting promoted to 64 bit.. probably because of the signed comparison below?
 		int32 GeneratedMipCount = FImageCoreUtils::GetMipCountFromDimensions(BaseSizeX, BaseSizeY, BaseSizeZ, this->bVolume);
@@ -3784,7 +3761,7 @@ bool FTextureBuildSettings::GetOutputMipInfo(
 	}
 	else
 	{
-		uint32 LongLatCubemapExtents = ComputeLongLatCubemapExtents(BaseSizeX, this->MaxTextureResolution);
+		uint32 LongLatCubemapExtents = UE::TextureBuildUtilities::ComputeLongLatCubemapExtents(BaseSizeX, this->MaxTextureResolution);
 		BaseSizeX = LongLatCubemapExtents;
 		BaseSizeY = LongLatCubemapExtents;
 		OutMip0NumSlices = 6 * InMip0NumSlices;
@@ -4163,26 +4140,17 @@ private:
 
 			if (bPadOrStretchTexture)
 			{
-				bool bResizeTexture = PowerOfTwoMode == ETexturePowerOfTwoSetting::StretchToPowerOfTwo || PowerOfTwoMode == ETexturePowerOfTwoSetting::StretchToSquarePowerOfTwo || PowerOfTwoMode == ETexturePowerOfTwoSetting::ResizeToSpecificResolution;
 				if (BuildSettings.MipGenSettings == TMGS_LeaveExistingMips)
 				{
 					// pad/stretch+leave existing is broken
 					UE_LOG(LogTextureCompressor, Error,	TEXT("Texture padding or resizing is not allowed when leaving existing mips."));
 					return false;
 				}
-				if ( bLongLatCubemap )
-				{
-					if (bResizeTexture)
-					{
-						UE_LOG(LogTextureCompressor, Warning, TEXT("In order to improve the quality of the generated texture, resizing LongLat cubemaps should be avoided."));
-					}
-					else
-					{
-						UE_LOG(LogTextureCompressor, Warning, TEXT("Padding of a LongLat cubemap may result in incorrect mapping of the texture pixels to the cubemap faces and should be avoided."));
-					}
-				}
 
 				// Want to stretch or pad the texture
+				bool bResizeTexture = PowerOfTwoMode == ETexturePowerOfTwoSetting::StretchToPowerOfTwo || PowerOfTwoMode == ETexturePowerOfTwoSetting::StretchToSquarePowerOfTwo || PowerOfTwoMode == ETexturePowerOfTwoSetting::ResizeToSpecificResolution;
+				// if not bResizeTexture, then padding
+
 				bool bSuitableFormat = FirstSourceMipImage.Format == ERawImageFormat::RGBA32F;
 
 				FImage Temp;
@@ -4200,7 +4168,12 @@ private:
 				{		
 					// changed resizer for "stretch to power of two" , bumped ddc key
 
-					// choice of Filter?
+					// choice of Filter? wrap or clamp?
+
+					// if you have multiple resize operations on a texture, such as StrechToPow2 and then also MaxTextureSize or Downscale
+					//	currently those are done one by one, which is not great
+					//	would be better to compute the net output size of all the operations
+					//	and do just one resize to get directly from source to final size
 
 					FImageCore::ResizeImageAllocDest(SourceImage,TargetImage, TargetTextureSizeX, TargetTextureSizeY);
 				}
@@ -4363,6 +4336,7 @@ private:
 
 					// @todo : find closest ResizeFilter that matches MipGen settings (maybe?)
 					//	 I mean, maybe not?  Most of the time MipGen is "SimpleAverage" which produces a really bad resize
+					//	-> wrap/clamp ?
 					FImageCore::EResizeImageFilter ResizeFilter = FImageCore::EResizeImageFilter::Default;
 					FImageCore::ResizeImage(BaseImage,DestImage,ResizeFilter);
 				}
@@ -4392,6 +4366,7 @@ private:
 					Temp.RawData.Empty();					
 					pSourceMips->Empty();
 
+					// note: on volumes, only XY size is checked, but Z size changes too
 					while( BuildSourceImageMips.Last().SizeX > MaxTextureResolution || 
 						   BuildSourceImageMips.Last().SizeY > MaxTextureResolution )
 					{
@@ -4439,7 +4414,8 @@ private:
 			//  eg. 256 makes 9 mips , 300 also makes 9 mips
 			if (bLongLatCubemap)
 			{
-				NumOutputMips = FImageCoreUtils::GetMipCountFromDimensions(ComputeLongLatCubemapExtents(TopMip.SizeX, BuildSettings.MaxTextureResolution), 1, 1, false);
+				uint32 Extents = UE::TextureBuildUtilities::ComputeLongLatCubemapExtents(TopMip.SizeX, BuildSettings.MaxTextureResolution);
+				NumOutputMips = FImageCoreUtils::GetMipCountFromDimensions(Extents, Extents, 6, false);
 			}
 			else
 			{
diff --git a/Engine/Source/Runtime/Engine/Private/Texture.cpp b/Engine/Source/Runtime/Engine/Private/Texture.cpp
--- a/Engine/Source/Runtime/Engine/Private/Texture.cpp
+++ b/Engine/Source/Runtime/Engine/Private/Texture.cpp
@@ -623,7 +623,12 @@ void UTexture::ValidateSettingsAfterImportOrEdit(bool * pRequiresNotifyMaterials
 			PowerOfTwoMode = ETexturePowerOfTwoSetting::None;
 		}
 
-		if ((PowerOfTwoMode == ETexturePowerOfTwoSetting::StretchToPowerOfTwo || PowerOfTwoMode == ETexturePowerOfTwoSetting::StretchToSquarePowerOfTwo || PowerOfTwoMode == ETexturePowerOfTwoSetting::ResizeToSpecificResolution) && !this->IsA<UTexture2D>())
+		// PadToPow2 for CubeMaps will almost never do something useful, but go ahead and allow it
+		// PowerOfTwo actions on LongLat CubeMaps act on the source *before* converting to a cube
+		//	 which is pretty pointless (output cube will always be pow2 anyway)
+		//	 but again, allow it if it's requested
+		
+		if ( PowerOfTwoMode == ETexturePowerOfTwoSetting::ResizeToSpecificResolution && !this->IsA<UTexture2D>() )
 		{
 			// currently resizing is only supported for 2D textures, but can be implemented for other types of textures in the future
 			UE_LOG(LogTexture, Display, TEXT("Currently resizing is only supported for Texture2D, forcing PowerOfTwoMode to None. (%s)"), *GetName());
@@ -3904,8 +3909,7 @@ void UTexture::GetBuiltTextureSize(const ITargetPlatformSettings* TargetPlatform
 
 	if (Source.IsLongLatCubemap())
 	{
-		// this should be kept in sync with ComputeLongLatCubemapExtents()
-		SizeX = SizeY = FMath::Max(1 << FMath::FloorLog2(SizeX / 2), 32);
+		SizeX = SizeY = UE::TextureBuildUtilities::ComputeLongLatCubemapExtents(SizeX, MaxTextureSize);
 	}
 
 	//we need to really have the actual top mip size of output platformdata
diff --git a/Engine/Source/Runtime/Engine/Private/TextureDerivedDataTask.cpp b/Engine/Source/Runtime/Engine/Private/TextureDerivedDataTask.cpp
--- a/Engine/Source/Runtime/Engine/Private/TextureDerivedDataTask.cpp
+++ b/Engine/Source/Runtime/Engine/Private/TextureDerivedDataTask.cpp
@@ -269,26 +269,45 @@ void FTextureSourceData::Init(UTexture& InTexture, TextureMipGenSettings InMipGe
 		}
 	}
 
-	for (FTextureSourceBlockData& Block : Blocks)
+	if ( Blocks.Num() == 0 )
 	{
-		int32 AdjustedSizeX, AdjustedSizeY, AdjustedSizeZ;
+		UE_LOG(LogTexture, Error, TEXT("No valid source blocks [%s]"), *InTexture.GetPathName());
+		check( bValid == false );
+		return;
+	}
+
+	if ( Blocks.Num() > 1 )
+	{
+		int32 BlockSizeZ=1;
 		UE::TextureBuildUtilities::GetPowerOfTwoTargetTextureSize(
-			Block.SizeX, Block.SizeY, 1,
+			BlockSizeX, BlockSizeY, 1,
 			false,
 			InPow2Setting, InResizeDuringBuildX, InResizeDuringBuildY,
-			AdjustedSizeX, AdjustedSizeY, AdjustedSizeZ);
+			BlockSizeX, BlockSizeY, BlockSizeZ);
 
-		// for the common case of NumBlocks == 1, BlockSizeX == Block.SizeX, MipBiasX/Y will both be zero
-		const int32 MipBiasX = FMath::CeilLogTwo(BlockSizeX / AdjustedSizeX);
-		const int32 MipBiasY = FMath::CeilLogTwo(BlockSizeY / AdjustedSizeY);
-		if (MipBiasX != MipBiasY)
+		for (FTextureSourceBlockData& Block : Blocks)
 		{
-			// @todo Oodle: this is failing even if "pad to pow2 square" is set, can we allow it through in that case?
-			UE_LOG(LogTexture, Warning, TEXT("VT has blocks with mismatched aspect ratios, cannot build. [%s]"), *InTexture.GetPathName());  // <- should be an Error, not a Warning
-			return;
-		}
+			int32 AdjustedSizeX, AdjustedSizeY, AdjustedSizeZ;
+			UE::TextureBuildUtilities::GetPowerOfTwoTargetTextureSize(
+				Block.SizeX, Block.SizeY, 1,
+				false,
+				InPow2Setting, InResizeDuringBuildX, InResizeDuringBuildY,
+				AdjustedSizeX, AdjustedSizeY, AdjustedSizeZ);
 
-		Block.MipBias = MipBiasX;
+			// for the common case of NumBlocks == 1, BlockSizeX == Block.SizeX, MipBiasX/Y will both be zero
+			const int32 MipBiasX = FMath::CeilLogTwo(BlockSizeX / AdjustedSizeX);
+			const int32 MipBiasY = FMath::CeilLogTwo(BlockSizeY / AdjustedSizeY);
+			if (MipBiasX != MipBiasY)
+			{
+				// @todo Oodle: this is failing even if "pad to pow2 square" is set, can we allow it through in that case?
+				// @@!! is this fixed now?
+				UE_LOG(LogTexture, Error, TEXT("VT has blocks with mismatched aspect ratios, cannot build. [%s]"), *InTexture.GetPathName());
+				check( bValid == false );
+				return;
+			}
+
+			Block.MipBias = MipBiasX;
+		}
 	}
 
 	TextureFullName = InTexture.GetFullName();
diff --git a/Engine/Source/Runtime/TextureUtilitiesCommon/Private/TextureSourceDataUtils.cpp b/Engine/Source/Runtime/TextureUtilitiesCommon/Private/TextureSourceDataUtils.cpp
--- a/Engine/Source/Runtime/TextureUtilitiesCommon/Private/TextureSourceDataUtils.cpp
+++ b/Engine/Source/Runtime/TextureUtilitiesCommon/Private/TextureSourceDataUtils.cpp
@@ -9,6 +9,7 @@
 #include "HAL/UnrealMemory.h"
 #include "EngineLogs.h"
 #include "TextureImportSettings.h"
+#include "TextureBuildUtilities.h"
 #include "ImageUtils.h"
 
 namespace UE::TextureUtilitiesCommon::Experimental
@@ -546,8 +547,7 @@ TEXTUREUTILITIESCOMMON_API bool DownsizeTextureSourceDataNearRenderingSize(UText
 	if (Texture->Source.IsLongLatCubemap())
 	{
 		// The function return the max size of the generated cube from the source long lat
-		// this should be kept in sync with the implementation details of ComputeLongLatCubemapExtents() or refactored
-		TargetSourceSize = (1U << FMath::FloorLog2(TargetSizeInGame)) * 2;
+		TargetSourceSize = UE::TextureBuildUtilities::ComputeLongLatCubemapExtents(TargetSizeInGame,0);
 	}
 
 	TargetSourceSize = FMath::Min(TargetSourceSize,AdditionalSourceSizeLimit);
diff --git a/Engine/Source/Runtime/TextureUtilitiesCommon/TextureUtilitiesCommon.build.cs b/Engine/Source/Runtime/TextureUtilitiesCommon/TextureUtilitiesCommon.build.cs
--- a/Engine/Source/Runtime/TextureUtilitiesCommon/TextureUtilitiesCommon.build.cs
+++ b/Engine/Source/Runtime/TextureUtilitiesCommon/TextureUtilitiesCommon.build.cs
@@ -19,6 +19,7 @@ public class TextureUtilitiesCommon : ModuleRules
 				"ImageCore",
 				"Slate",
 				"SlateCore",
+				"TextureBuildUtilities",
 			});
 	}
 }
