Account for RGBA32F Source Data size in texture build memory estimates

When resizing, we don't just end up with the _target_ image in RGBA32F, but also (at least temporarily) the source.

diff --git a/Engine/Source/Developer/TextureBuildUtilities/Private/TextureBuildUtilities.cpp b/Engine/Source/Developer/TextureBuildUtilities/Private/TextureBuildUtilities.cpp
--- a/Engine/Source/Developer/TextureBuildUtilities/Private/TextureBuildUtilities.cpp
+++ b/Engine/Source/Developer/TextureBuildUtilities/Private/TextureBuildUtilities.cpp
@@ -557,7 +557,11 @@ TEXTUREBUILDUTILITIES_API int64 GetVirtualTextureRequiredMemoryEstimate(const FT
 		}
 		else if ( bDidPow2 )
 		{
-			// surface in in RGBA32F
+			// We create a copy of the source in RGBA32F as part of resizing
+			// for some formats (FirstSourceMipImage "convert to RGBA32F")
+			int64 SourceBlockNumPixels = (int64)SourceBlock.SizeX * SourceBlock.SizeY * SourceBlock.NumSlices;
+			ResizingPhaseMemUsePerLayer += SourceBlockNumPixels * 16;
+			// surface input to remaining processing in RGBA32F
 			ResizingPhaseMemUsePerLayer += AfterPow2TopMipNumPixels * 16;
 		}
 
@@ -718,9 +722,18 @@ TEXTUREBUILDUTILITIES_API int64 GetPhysicalTextureBuildMemoryEstimate(const FTex
 		BuildSettings.ResizeDuringBuildX, BuildSettings.ResizeDuringBuildY,
 		TargetSizeX, TargetSizeY, TargetSizeZ);
 	
-	// mem use of the Pow2 op is always lower that later ops, so no need to compute it
-
 	int64 ResizingPhaseMemUse = TotalSourceBytes;
+
+	// Pow2 resize can end up converting the *source* data to RGBA32F, so we need to account for it
+	if (bDidPow2)
+	{
+		//  FirstSourceMipImage "convert to RGBA32F" to FImage Temp in TextureCompressorModule
+		int64 SourceDataMipNumPixels = (int64)InSourceImageInfo.SizeX * InSourceImageInfo.SizeY * InSourceImageInfo.NumSlices;
+		ResizingPhaseMemUse += SourceDataMipNumPixels * 16; // original source data in RGBA32F
+		// This is live concurrently with the top mip in RGBA32F computed next.
+		// Therefore we need the sum of both, not the max.
+	}
+
 	int64 InitialTopMipNumPixels = (int64) TargetSizeX * TargetSizeY * TargetSizeZ;
 	ResizingPhaseMemUse += InitialTopMipNumPixels * 16; // top mip in RGBA32F may be needed
 
