Avoid hard crash when texture format dll is missing.

#rnx

diff --git a/Engine/Source/Developer/TextureBuildUtilities/Private/TextureBuildUtilities.cpp b/Engine/Source/Developer/TextureBuildUtilities/Private/TextureBuildUtilities.cpp
--- a/Engine/Source/Developer/TextureBuildUtilities/Private/TextureBuildUtilities.cpp
+++ b/Engine/Source/Developer/TextureBuildUtilities/Private/TextureBuildUtilities.cpp
@@ -453,6 +453,11 @@ TEXTUREBUILDUTILITIES_API EPixelFormat GetOutputPixelFormatWithFallback(const FT
 		}
 	}
 
+	if (InBuildSettings.BaseTextureFormat == nullptr)
+	{
+		return PF_Unknown;
+	}
+
 	EPixelFormat PixelFormat = PF_Unknown;
 	if (bNeedsBaseCopy)
 	{
@@ -695,6 +700,12 @@ TEXTUREBUILDUTILITIES_API int64 GetVirtualTextureRequiredMemoryEstimate(const FT
 // Returns the estimated memory cost of building the given texture. Only valid for physical (i.e. non-virtual) textures.
 TEXTUREBUILDUTILITIES_API int64 GetPhysicalTextureBuildMemoryEstimate(const FTextureBuildSettings* InSettingsPerLayerFetchFirst, const FImageInfo& InSourceImageInfo, int32 InMipCount)
 {
+	if (InSettingsPerLayerFetchFirst->BaseTextureFormat == nullptr)
+	{
+		// Will fail build later, return no memory estimate
+		return -1;
+	}
+
 	const bool bRDO = true;
 	// @todo Oodle : be careful about using BuildSettings for bRDO as there are two buildsettingses, just assume its on for now
 	//   <- FIX ME, allow lower mem estimates for non-RDO
diff --git a/Engine/Source/Runtime/Engine/Private/TextureDerivedData.cpp b/Engine/Source/Runtime/Engine/Private/TextureDerivedData.cpp
--- a/Engine/Source/Runtime/Engine/Private/TextureDerivedData.cpp
+++ b/Engine/Source/Runtime/Engine/Private/TextureDerivedData.cpp
@@ -727,16 +727,15 @@ static void GetBuiltTextureSizeBytesEstimate(
 //	OutSettings must be otherwise fully set up
 static void ModifyMaxTextureResolutionBuildSettingsForPlatformLimit(
 	const UTexture& Texture, 
-	const ITargetPlatform* TargetPlatform,
-	
+	const ITargetPlatform* TargetPlatform,	
 	const ITextureFormat* TextureFormat,
-
 	FTextureBuildSettings& OutSettings)
 {
 	check( ! OutSettings.bVirtualStreamable );
 	check( OutSettings.TextureFormatName != NAME_None );
 	
-	if (!Texture.Source.IsValid())
+	if (!Texture.Source.IsValid() ||
+		OutSettings.BaseTextureFormat == nullptr) // can happen with missing format dlls.
 	{
 		// Nothing to do - texture can't be built.
 		return;
